Initial submit of level prototyping assets in template resources.

diff --git a/Templates/TemplateResources/High/LevelPrototyping/FeaturePack/manifest.json b/Templates/TemplateResources/High/LevelPrototyping/FeaturePack/manifest.json
new file mode 100644
--- /dev/null
+++ b/Templates/TemplateResources/High/LevelPrototyping/FeaturePack/manifest.json
@@ -0,0 +1,48 @@
+{
+	"Version": 1,
+	"Name":
+	[
+		{
+			"Language": "en",
+			"Text": "Level Prototyping Pack"
+		}
+	],
+	"Description":
+	[
+		{
+			"Language": "en",
+			"Text": "Level prototyping meshes and materials for quickly blocking out and iterating on level design."
+		}
+	],
+	"AssetTypes":
+	[
+		{
+			"Language": "en",
+			"Text": "StaticMesh, Material"
+		}
+	],
+	"SearchTags":
+	[
+		{
+			"Language": "en",
+			"Text": "Level prototyping, Static Meshes, Materials"
+		}
+	],	
+	"ClassTypes": "StaticMesh",
+	"Category": "SharedPack",
+	"Thumbnail": "LevelPrototyping.png",
+	"Screenshots":
+	[
+		"LevelPrototyping_preview.png"
+	],
+	"Ident": "LevelPrototypingHigh",
+	"AdditionalFiles":
+	{
+		"DestinationFilesFolder": "LevelPrototyping",
+		"AdditionalFilesList":
+		[
+			"Templates/TemplateResources/High/LevelPrototyping/Content/*.*"
+		]
+	}
+}
+
