Template update:
- Added shared folder "Characters" to TemplateResources
- Added Manny and Quinn to Top Down and Third Person templates
- 3rd Person default: Quinn
- Top Down default: Manny

diff --git a/Templates/TP_ThirdPerson/Config/TemplateDefs.ini b/Templates/TP_ThirdPerson/Config/TemplateDefs.ini
--- a/Templates/TP_ThirdPerson/Config/TemplateDefs.ini
+++ b/Templates/TP_ThirdPerson/Config/TemplateDefs.ini
@@ -49,5 +49,5 @@ ReplacementsInFiles=(Extensions=("cpp","h","ini","cs"),From="%TEMPLATENAME_UPPER
 ReplacementsInFiles=(Extensions=("cpp","h","ini","cs"),From="%TEMPLATENAME_LOWERCASE%",To="%PROJECTNAME_LOWERCASE%",bCaseSensitive=true)
 ReplacementsInFiles=(Extensions=("cpp","h","ini","cs"),From="%TEMPLATENAME%",To="%PROJECTNAME%",bCaseSensitive=false)
 SharedContentPacks=(MountName="LevelPrototyping",DetailLevels=("High"))
-SharedContentPacks=(MountName="Mannequin",DetailLevels=("High"))
+SharedContentPacks=(MountName="Characters",DetailLevels=("High"))
 EditDetailLevelPreference="High"
\ No newline at end of file
diff --git a/Templates/TP_ThirdPerson/manifest.json b/Templates/TP_ThirdPerson/manifest.json
--- a/Templates/TP_ThirdPerson/manifest.json
+++ b/Templates/TP_ThirdPerson/manifest.json
@@ -70,7 +70,7 @@
 			]
 		},
 		{
-			"MountName": "Mannequin",
+			"MountName": "Characters",
 			"DetailLevels":
 			[		
 				"High"
diff --git a/Templates/TP_ThirdPersonBP/Config/TemplateDefs.ini b/Templates/TP_ThirdPersonBP/Config/TemplateDefs.ini
--- a/Templates/TP_ThirdPersonBP/Config/TemplateDefs.ini
+++ b/Templates/TP_ThirdPersonBP/Config/TemplateDefs.ini
@@ -42,5 +42,5 @@ ReplacementsInFiles=(Extensions=("cpp","h","ini","cs"),From="%TEMPLATENAME_UPPER
 ReplacementsInFiles=(Extensions=("cpp","h","ini","cs"),From="%TEMPLATENAME_LOWERCASE%",To="%PROJECTNAME_LOWERCASE%",bCaseSensitive=true)
 ReplacementsInFiles=(Extensions=("cpp","h","ini","cs"),From="%TEMPLATENAME%",To="%PROJECTNAME%",bCaseSensitive=false)
 SharedContentPacks=(MountName="LevelPrototyping",DetailLevels=("High"))
-SharedContentPacks=(MountName="Mannequin",DetailLevels=("High"))
+SharedContentPacks=(MountName="Characters",DetailLevels=("High"))
 EditDetailLevelPreference="High"
diff --git a/Templates/TP_ThirdPersonBP/manifest.json b/Templates/TP_ThirdPersonBP/manifest.json
--- a/Templates/TP_ThirdPersonBP/manifest.json
+++ b/Templates/TP_ThirdPersonBP/manifest.json
@@ -70,7 +70,7 @@
 			]
 		},
 		{
-			"MountName": "Mannequin",
+			"MountName": "Characters",
 			"DetailLevels":
 			[
 				"High"
diff --git a/Templates/TP_TopDown/Config/TemplateDefs.ini b/Templates/TP_TopDown/Config/TemplateDefs.ini
--- a/Templates/TP_TopDown/Config/TemplateDefs.ini
+++ b/Templates/TP_TopDown/Config/TemplateDefs.ini
@@ -48,6 +48,6 @@ ReplacementsInFiles=(Extensions=("cpp","h","ini","cs"),From="%TEMPLATENAME_UPPER
 ReplacementsInFiles=(Extensions=("cpp","h","ini","cs"),From="%TEMPLATENAME_LOWERCASE%",To="%PROJECTNAME_LOWERCASE%",bCaseSensitive=true)
 ReplacementsInFiles=(Extensions=("cpp","h","ini","cs"),From="%TEMPLATENAME%",To="%PROJECTNAME%",bCaseSensitive=false)SharedContentPacks=(MountName="TopDown",DetailLevels=("Standard"))
 SharedContentPacks=(MountName="LevelPrototyping",DetailLevels=("High"))
-SharedContentPacks=(MountName="Mannequin",DetailLevels=("High"))
+SharedContentPacks=(MountName="Characters",DetailLevels=("High"))
 SharedContentPacks=(MountName="Cursor",DetailLevels=("Standard"))
-EditDetailLevelPreference="High"
\ No newline at end of file
+EditDetailLevelPreference="High"
diff --git a/Templates/TP_TopDown/manifest.json b/Templates/TP_TopDown/manifest.json
--- a/Templates/TP_TopDown/manifest.json
+++ b/Templates/TP_TopDown/manifest.json
@@ -77,7 +77,7 @@
 			]
 		},
 		{
-			"MountName": "Mannequin",
+			"MountName": "Characters",
 			"DetailLevels":
 			[
 				"High"
diff --git a/Templates/TP_TopDownBP/Config/TemplateDefs.ini b/Templates/TP_TopDownBP/Config/TemplateDefs.ini
--- a/Templates/TP_TopDownBP/Config/TemplateDefs.ini
+++ b/Templates/TP_TopDownBP/Config/TemplateDefs.ini
@@ -43,6 +43,6 @@ ReplacementsInFiles=(Extensions=("cpp","h","ini","cs"),From="%TEMPLATENAME_UPPER
 ReplacementsInFiles=(Extensions=("cpp","h","ini","cs"),From="%TEMPLATENAME_LOWERCASE%",To="%PROJECTNAME_LOWERCASE%",bCaseSensitive=true)
 ReplacementsInFiles=(Extensions=("cpp","h","ini","cs"),From="%TEMPLATENAME%",To="%PROJECTNAME%",bCaseSensitive=false)
 SharedContentPacks=(MountName="LevelPrototyping",DetailLevels=("High"))
-SharedContentPacks=(MountName="Mannequin",DetailLevels=("High"))
+SharedContentPacks=(MountName="Characters",DetailLevels=("High"))
 SharedContentPacks=(MountName="Cursor",DetailLevels=("Standard"))
 EditDetailLevelPreference="High"
diff --git a/Templates/TP_TopDownBP/manifest.json b/Templates/TP_TopDownBP/manifest.json
--- a/Templates/TP_TopDownBP/manifest.json
+++ b/Templates/TP_TopDownBP/manifest.json
@@ -77,7 +77,7 @@
 			]
 		},
 		{
-			"MountName": "Mannequin",
+			"MountName": "Characters",
 			"DetailLevels":
 			[
 				"High"
diff --git a/Templates/TemplateResources/High/Characters/FeaturePack/manifest.json b/Templates/TemplateResources/High/Characters/FeaturePack/manifest.json
new file mode 100644
--- /dev/null
+++ b/Templates/TemplateResources/High/Characters/FeaturePack/manifest.json
@@ -0,0 +1,48 @@
+{
+	"Version": 1,
+	"Name":
+	[
+		{
+			"Language": "en",
+			"Text": "Mannequin shared resource high"
+		}
+	],
+	"Description":
+	[
+		{
+			"Language": "en",
+			"Text": "Shared resource pack for high Mannequin"
+		}
+	],
+	"AssetTypes":
+	[
+		{
+			"Language": "en",
+			"Text": "SkeletonMesh, Skeleton, Material, Animation, IK, Retargeting"
+		}
+	],
+	"SearchTags":
+	[
+		{
+			"Language": "en",
+			"Text": "Mannequin, Manny, Quinn, SkeletonMesh, Skeleton, Material, Animation, IK, Retargeting"
+		}
+	],	
+	"ClassTypes": "SkeletonMeshes",
+	"Category": "SharedPack",
+	"Thumbnail": "Characters.png",
+	"Screenshots":
+	[
+		"Characters_preview.png"
+	],
+	"Ident": "MannequinHigh",
+	"AdditionalFiles":
+	{
+		"DestinationFilesFolder": "Characters",
+		"AdditionalFilesList":
+		[
+			"Templates/TemplateResources/High/Characters/Content/*.*"
+		]
+	}
+}
+
