Templates: fixed wrong ident in manifest which lead to an issue not being able to create a project when new mannequins were referenced

diff --git a/Templates/TemplateResources/High/Characters/FeaturePack/manifest.json b/Templates/TemplateResources/High/Characters/FeaturePack/manifest.json
--- a/Templates/TemplateResources/High/Characters/FeaturePack/manifest.json
+++ b/Templates/TemplateResources/High/Characters/FeaturePack/manifest.json
@@ -28,14 +28,14 @@
 			"Text": "Mannequin, Manny, Quinn, SkeletonMesh, Skeleton, Material, Animation, IK, Retargeting"
 		}
 	],	
-	"ClassTypes": "SkeletonMeshes",
+	"ClassTypes": "StaticMesh",
 	"Category": "SharedPack",
 	"Thumbnail": "Characters.png",
 	"Screenshots":
 	[
 		"Characters_preview.png"
 	],
-	"Ident": "MannequinHigh",
+	"Ident": "CharactersHigh",
 	"AdditionalFiles":
 	{
 		"DestinationFilesFolder": "Characters",
