Update folder structure
Created input manifest
Updated UI on third person combat
mobile touch support

diff --git a/Templates/TP_FirstPersonBP/Config/DefaultEngine.ini b/Templates/TP_FirstPersonBP/Config/DefaultEngine.ini
--- a/Templates/TP_FirstPersonBP/Config/DefaultEngine.ini
+++ b/Templates/TP_FirstPersonBP/Config/DefaultEngine.ini
@@ -56,7 +56,7 @@ AppliedDefaultGraphicsPerformance=Maximum
 [/Script/Engine.Engine]
 +ActiveGameNameRedirects=(OldGameName="TP_FirstPersonBP",NewGameName="/Script/MyProject")
 +ActiveGameNameRedirects=(OldGameName="/Script/TP_FirstPersonBP",NewGameName="/Script/MyProject")
-NearClipPlane=0.500000
+NearClipPlane=5.000000
 
 [/Script/AndroidFileServerEditor.AndroidFileServerRuntimeSettings]
 bEnablePlugin=True
@@ -75,3 +75,6 @@ ManualIPAddress=
 [/Script/AIModule.AISystem]
 bForgetStaleActors=True
 
+[/Script/Engine.UserInterfaceSettings]
+UIScaleRule=ShortestSide
+
diff --git a/Templates/TP_FirstPersonBP/Config/DefaultInput.ini b/Templates/TP_FirstPersonBP/Config/DefaultInput.ini
--- a/Templates/TP_FirstPersonBP/Config/DefaultInput.ini
+++ b/Templates/TP_FirstPersonBP/Config/DefaultInput.ini
@@ -111,7 +111,10 @@ FOVScale=0.011110
 DoubleClickTime=0.200000
 DefaultPlayerInputClass=/Script/EnhancedInput.EnhancedPlayerInput
 DefaultInputComponentClass=/Script/EnhancedInput.EnhancedInputComponent
-DefaultTouchInterface=/Game/FirstPerson/Input/MobileControls.MobileControls
+DefaultTouchInterface=None
 -ConsoleKeys=Tilde
 +ConsoleKeys=Tilde
 
+[EnhancedInputPlatformSettings_Windows EnhancedInputPlatformSettings]
+bShouldLogMappingContextRedirects=False
+
diff --git a/Templates/TP_FirstPersonBP/Config/TemplateDefs.ini b/Templates/TP_FirstPersonBP/Config/TemplateDefs.ini
--- a/Templates/TP_FirstPersonBP/Config/TemplateDefs.ini
+++ b/Templates/TP_FirstPersonBP/Config/TemplateDefs.ini
@@ -43,6 +43,7 @@ ReplacementsInFiles=(Extensions=("cpp","h","ini","cs"),From="%TEMPLATENAME_LOWER
 ReplacementsInFiles=(Extensions=("cpp","h","ini","cs"),From="%TEMPLATENAME%",To="%PROJECTNAME%",bCaseSensitive=false)
 SharedContentPacks=(MountName="Characters",DetailLevels=("High"))
 SharedContentPacks=(MountName="LevelPrototyping",DetailLevels=("High"))
+SharedContentPacks=(MountName="Input",DetailLevels=("High"))
 EditDetailLevelPreference="High"
 
 ; Variant packs
diff --git a/Templates/TP_VehicleAdvBP/Config/DefaultEngine.ini b/Templates/TP_VehicleAdvBP/Config/DefaultEngine.ini
--- a/Templates/TP_VehicleAdvBP/Config/DefaultEngine.ini
+++ b/Templates/TP_VehicleAdvBP/Config/DefaultEngine.ini
@@ -6,14 +6,14 @@ bSubstepping=False
 bSubsteppingAsync=True
 
 [/Script/EngineSettings.GameMapsSettings]
-EditorStartupMap=/Game/VehicleTemplate/Maps/VehicleBasic.VehicleBasic
+EditorStartupMap=/Game/VehicleTemplate/Maps/Lvl_VehicleBasic.Lvl_VehicleBasic
 LocalMapOptions=
 TransitionMap=
 bUseSplitscreen=True
 TwoPlayerSplitscreenLayout=Horizontal
 ThreePlayerSplitscreenLayout=FavorTop
 GameInstanceClass=/Script/Engine.GameInstance
-GameDefaultMap=/Game/VehicleTemplate/Maps/VehicleBasic.VehicleBasic
+GameDefaultMap=/Game/VehicleTemplate/Maps/Lvl_VehicleBasic.Lvl_VehicleBasic
 ServerDefaultMap=/Engine/Maps/Entry.Entry
 GlobalDefaultGameMode=/Game/VehicleTemplate/Blueprints/BP_GM_VehicleBasic.BP_GM_VehicleBasic_C
 GlobalDefaultServerGameMode=None
diff --git a/Templates/TemplateResources/High/Input/FeaturePack/manifest.json b/Templates/TemplateResources/High/Input/FeaturePack/manifest.json
new file mode 100644
--- /dev/null
+++ b/Templates/TemplateResources/High/Input/FeaturePack/manifest.json
@@ -0,0 +1,48 @@
+{
+	"Version": 1,
+	"Name":
+	[
+		{
+			"Language": "en",
+			"Text": "Mannequin shared resource high"
+		}
+	],
+	"Description":
+	[
+		{
+			"Language": "en",
+			"Text": "Shared resource pack for high Mannequin"
+		}
+	],
+	"AssetTypes":
+	[
+		{
+			"Language": "en",
+			"Text": "SkeletonMesh, Skeleton, Material, Animation, IK, Retargeting"
+		}
+	],
+	"SearchTags":
+	[
+		{
+			"Language": "en",
+			"Text": "Mannequin, Manny, Quinn, SkeletonMesh, Skeleton, Material, Animation, IK, Retargeting"
+		}
+	],	
+	"ClassTypes": "StaticMesh",
+	"Category": "SharedPack",
+	"Thumbnail": "Characters.png",
+	"Screenshots":
+	[
+		"Characters_preview.png"
+	],
+	"Ident": "InputHigh",
+	"AdditionalFiles":
+	{
+		"DestinationFilesFolder": "Input",
+		"AdditionalFilesList":
+		[
+			"Templates/TemplateResources/High/Input/Content/*.*"
+		]
+	}
+}
+
diff --git a/Templates/TemplateResources/Standard/FPWeapon/FeaturePack/manifest.json b/Templates/TemplateResources/Standard/FPWeapon/FeaturePack/manifest.json
deleted file mode 100644
diff --git a/Templates/TemplateResources/Standard/FirstPersonArms/FeaturePack/manifest.json b/Templates/TemplateResources/Standard/FirstPersonArms/FeaturePack/manifest.json
deleted file mode 100644
